Added Intro to UE Template

#virtualized

diff --git a/Templates/TP_UEIntro_BP/Config/DefaultEditorPerProjectUserSettings.ini b/Templates/TP_UEIntro_BP/Config/DefaultEditorPerProjectUserSettings.ini
new file mode 100644
--- /dev/null
+++ b/Templates/TP_UEIntro_BP/Config/DefaultEditorPerProjectUserSettings.ini
@@ -0,0 +1,2 @@
+[ContentBrowser]
+ContentBrowserTab1.SelectedPaths=/Game/FirstPersonBP
\ No newline at end of file
diff --git a/Templates/TP_UEIntro_BP/Config/DefaultEngine.ini b/Templates/TP_UEIntro_BP/Config/DefaultEngine.ini
new file mode 100644
--- /dev/null
+++ b/Templates/TP_UEIntro_BP/Config/DefaultEngine.ini
@@ -0,0 +1,93 @@
+[/Script/Engine.CollisionProfile]
++Profiles=(Name="Projectile",CollisionEnabled=QueryOnly,ObjectTypeName="Projectile",CustomResponses=,HelpMessage="Preset for projectiles",bCanModify=True)
++DefaultChannelResponses=(Channel=ECC_GameTraceChannel1,Name="Projectile",DefaultResponse=ECR_Block,bTraceType=False,bStaticObject=False)
++EditProfiles=(Name="Trigger",CustomResponses=((Channel=Projectile, Response=ECR_Ignore)))
+
+[/Script/EngineSettings.GameMapsSettings]
+EditorStartupMap=/Game/DemoTemplate/_Core/Lvl_IntroRoom.Lvl_IntroRoom
+LocalMapOptions=
+TransitionMap=
+bUseSplitscreen=True
+TwoPlayerSplitscreenLayout=Horizontal
+ThreePlayerSplitscreenLayout=FavorTop
+GameInstanceClass=/Script/Engine.GameInstance
+GameDefaultMap=/Game/DemoTemplate/_Core/Lvl_IntroRoom.Lvl_IntroRoom
+ServerDefaultMap=/Engine/Maps/Entry
+GlobalDefaultGameMode=/Game/FirstPerson/Blueprints/BP_FirstPersonGameMode.BP_FirstPersonGameMode_C
+GlobalDefaultServerGameMode=None
+
+[/Script/Engine.RendererSettings]
+r.ReflectionMethod=1
+r.GenerateMeshDistanceFields=True
+r.DynamicGlobalIlluminationMethod=1
+r.Mobile.EnableNoPrecomputedLightingCSMShader=1
+r.DefaultFeature.AutoExposure.ExtendDefaultLuminanceRange=True
+r.DefaultFeature.AutoExposure.ExtendDefaultLuminanceRange=True
+r.DefaultFeature.AutoExposure.ExtendDefaultLuminanceRange=true
+r.Shadow.Virtual.Enable=1
+
+r.DefaultFeature.LocalExposure.HighlightContrastScale=0.8
+r.DefaultFeature.LocalExposure.HighlightContrastScale=0.8
+
+r.DefaultFeature.LocalExposure.ShadowContrastScale=0.8
+r.DefaultFeature.LocalExposure.ShadowContrastScale=0.8
+r.DefaultFeature.LightUnits=2
+r.AllowStaticLighting=False
+RedChromaticityCoordinate=(X=0.640000,Y=0.330000)
+GreenChromaticityCoordinate=(X=0.300000,Y=0.600000)
+BlueChromaticityCoordinate=(X=0.150000,Y=0.060000)
+WhiteChromaticityCoordinate=(X=0.312700,Y=0.329000)
+WorkingColorSpaceChoice=sRGB
+r.LegacyLuminanceFactors=False
+r.Lumen.ScreenTracingSource=0
+r.FirstPerson.SelfShadow=1
+
+r.SkinCache.CompileShaders=True
+r.RayTracing=True
+
+r.RayTracing.RayTracingProxies.ProjectEnabled=True
+r.Substrate=False
+r.Refraction.Blur=False
+r.Lumen.HardwareRayTracing=True
+r.Lumen.Reflections.HardwareRayTracing.Translucent.Refraction.EnableForProject=False
+
+[/Script/WindowsTargetPlatform.WindowsTargetSettings]
+DefaultGraphicsRHI=DefaultGraphicsRHI_DX12
+DefaultGraphicsRHI=DefaultGraphicsRHI_DX12
+-D3D12TargetedShaderFormats=PCD3D_SM5
++D3D12TargetedShaderFormats=PCD3D_SM6
+-D3D11TargetedShaderFormats=PCD3D_SM5
++D3D11TargetedShaderFormats=PCD3D_SM5
+
+[/Script/HardwareTargeting.HardwareTargetingSettings]
+TargetedHardwareClass=EHardwareClass::Desktop
+AppliedTargetedHardwareClass=Desktop
+DefaultGraphicsPerformance=EGraphicsPreset::Maximum
+AppliedDefaultGraphicsPerformance=Maximum
+
+[/Script/Engine.Engine]
+NearClipPlane=5.000000
+
++ActiveGameNameRedirects=(OldGameName="TP_FirstPersonBP",NewGameName="/Script/FirstTemplate")
++ActiveGameNameRedirects=(OldGameName="/Script/TP_FirstPersonBP",NewGameName="/Script/FirstTemplate")
+
+[/Script/AIModule.AISystem]
+bForgetStaleActors=True
+
+[/Script/Engine.UserInterfaceSettings]
+UIScaleRule=ShortestSide
+
+[/Script/AndroidFileServerEditor.AndroidFileServerRuntimeSettings]
+bEnablePlugin=True
+bAllowNetworkConnection=True
+SecurityToken=55490E894C07D04A1ACF51BB99492878
+bIncludeInShipping=False
+bAllowExternalStartInShipping=False
+bCompileAFSProject=False
+bUseCompression=False
+bLogFiles=False
+bReportStats=False
+ConnectionType=USBOnly
+bUseManualIPAddress=False
+ManualIPAddress=
+
diff --git a/Templates/TP_UEIntro_BP/Config/DefaultGame.ini b/Templates/TP_UEIntro_BP/Config/DefaultGame.ini
new file mode 100644
--- /dev/null
+++ b/Templates/TP_UEIntro_BP/Config/DefaultGame.ini
@@ -0,0 +1,3 @@
+[/Script/EngineSettings.GeneralProjectSettings]
+ProjectID=86950E854DF21460DD9D4F9DF11D7867
+ProjectName=First Person BP Game Template
diff --git a/Templates/TP_UEIntro_BP/Config/DefaultInput.ini b/Templates/TP_UEIntro_BP/Config/DefaultInput.ini
new file mode 100644
--- /dev/null
+++ b/Templates/TP_UEIntro_BP/Config/DefaultInput.ini
@@ -0,0 +1,120 @@
+[/Script/Engine.InputSettings]
+-AxisConfig=(AxisKeyName="Gamepad_LeftX",AxisProperties=(DeadZone=0.25,Exponent=1.f,Sensitivity=1.f))
+-AxisConfig=(AxisKeyName="Gamepad_LeftY",AxisProperties=(DeadZone=0.25,Exponent=1.f,Sensitivity=1.f))
+-AxisConfig=(AxisKeyName="Gamepad_RightX",AxisProperties=(DeadZone=0.25,Exponent=1.f,Sensitivity=1.f))
+-AxisConfig=(AxisKeyName="Gamepad_RightY",AxisProperties=(DeadZone=0.25,Exponent=1.f,Sensitivity=1.f))
+-AxisConfig=(AxisKeyName="MouseX",AxisProperties=(DeadZone=0.f,Exponent=1.f,Sensitivity=0.07f))
+-AxisConfig=(AxisKeyName="MouseY",AxisProperties=(DeadZone=0.f,Exponent=1.f,Sensitivity=0.07f))
+-AxisConfig=(AxisKeyName="Mouse2D",AxisProperties=(DeadZone=0.f,Exponent=1.f,Sensitivity=0.07f))
++AxisConfig=(AxisKeyName="Mouse2D",AxisProperties=(DeadZone=0.000000,Sensitivity=0.070000,Exponent=1.000000,bInvert=False))
++AxisConfig=(AxisKeyName="Gamepad_LeftX",AxisProperties=(DeadZone=0.250000,Sensitivity=1.000000,Exponent=1.000000,bInvert=False))
++AxisConfig=(AxisKeyName="Gamepad_LeftY",AxisProperties=(DeadZone=0.250000,Sensitivity=1.000000,Exponent=1.000000,bInvert=False))
++AxisConfig=(AxisKeyName="Gamepad_RightX",AxisProperties=(DeadZone=0.250000,Sensitivity=1.000000,Exponent=1.000000,bInvert=False))
++AxisConfig=(AxisKeyName="Gamepad_RightY",AxisProperties=(DeadZone=0.250000,Sensitivity=1.000000,Exponent=1.000000,bInvert=False))
++AxisConfig=(AxisKeyName="MouseX",AxisProperties=(DeadZone=0.000000,Sensitivity=0.070000,Exponent=1.000000,bInvert=False))
++AxisConfig=(AxisKeyName="MouseY",AxisProperties=(DeadZone=0.000000,Sensitivity=0.070000,Exponent=1.000000,bInvert=False))
++AxisConfig=(AxisKeyName="MouseWheelAxis",AxisProperties=(DeadZone=0.000000,Sensitivity=1.000000,Exponent=1.000000,bInvert=False))
++AxisConfig=(AxisKeyName="Gamepad_LeftTriggerAxis",AxisProperties=(DeadZone=0.000000,Sensitivity=1.000000,Exponent=1.000000,bInvert=False))
++AxisConfig=(AxisKeyName="Gamepad_RightTriggerAxis",AxisProperties=(DeadZone=0.000000,Sensitivity=1.000000,Exponent=1.000000,bInvert=False))
++AxisConfig=(AxisKeyName="Gamepad_Special_Left_X",AxisProperties=(DeadZone=0.000000,Sensitivity=1.000000,Exponent=1.000000,bInvert=False))
++AxisConfig=(AxisKeyName="Gamepad_Special_Left_Y",AxisProperties=(DeadZone=0.000000,Sensitivity=1.000000,Exponent=1.000000,bInvert=False))
++AxisConfig=(AxisKeyName="Vive_Left_Trigger_Axis",AxisProperties=(DeadZone=0.000000,Sensitivity=1.000000,Exponent=1.000000,bInvert=False))
++AxisConfig=(AxisKeyName="Vive_Left_Trackpad_X",AxisProperties=(DeadZone=0.000000,Sensitivity=1.000000,Exponent=1.000000,bInvert=False))
++AxisConfig=(AxisKeyName="Vive_Left_Trackpad_Y",AxisProperties=(DeadZone=0.000000,Sensitivity=1.000000,Exponent=1.000000,bInvert=False))
++AxisConfig=(AxisKeyName="Vive_Right_Trigger_Axis",AxisProperties=(DeadZone=0.000000,Sensitivity=1.000000,Exponent=1.000000,bInvert=False))
++AxisConfig=(AxisKeyName="Vive_Right_Trackpad_X",AxisProperties=(DeadZone=0.000000,Sensitivity=1.000000,Exponent=1.000000,bInvert=False))
++AxisConfig=(AxisKeyName="Vive_Right_Trackpad_Y",AxisProperties=(DeadZone=0.000000,Sensitivity=1.000000,Exponent=1.000000,bInvert=False))
++AxisConfig=(AxisKeyName="MixedReality_Left_Trigger_Axis",AxisProperties=(DeadZone=0.000000,Sensitivity=1.000000,Exponent=1.000000,bInvert=False))
++AxisConfig=(AxisKeyName="MixedReality_Left_Thumbstick_X",AxisProperties=(DeadZone=0.000000,Sensitivity=1.000000,Exponent=1.000000,bInvert=False))
++AxisConfig=(AxisKeyName="MixedReality_Left_Thumbstick_Y",AxisProperties=(DeadZone=0.000000,Sensitivity=1.000000,Exponent=1.000000,bInvert=False))
++AxisConfig=(AxisKeyName="MixedReality_Left_Trackpad_X",AxisProperties=(DeadZone=0.000000,Sensitivity=1.000000,Exponent=1.000000,bInvert=False))
++AxisConfig=(AxisKeyName="MixedReality_Left_Trackpad_Y",AxisProperties=(DeadZone=0.000000,Sensitivity=1.000000,Exponent=1.000000,bInvert=False))
++AxisConfig=(AxisKeyName="MixedReality_Right_Trigger_Axis",AxisProperties=(DeadZone=0.000000,Sensitivity=1.000000,Exponent=1.000000,bInvert=False))
++AxisConfig=(AxisKeyName="MixedReality_Right_Thumbstick_X",AxisProperties=(DeadZone=0.000000,Sensitivity=1.000000,Exponent=1.000000,bInvert=False))
++AxisConfig=(AxisKeyName="MixedReality_Right_Thumbstick_Y",AxisProperties=(DeadZone=0.000000,Sensitivity=1.000000,Exponent=1.000000,bInvert=False))
++AxisConfig=(AxisKeyName="MixedReality_Right_Trackpad_X",AxisProperties=(DeadZone=0.000000,Sensitivity=1.000000,Exponent=1.000000,bInvert=False))
++AxisConfig=(AxisKeyName="MixedReality_Right_Trackpad_Y",AxisProperties=(DeadZone=0.000000,Sensitivity=1.000000,Exponent=1.000000,bInvert=False))
++AxisConfig=(AxisKeyName="OculusTouch_Left_Grip_Axis",AxisProperties=(DeadZone=0.000000,Sensitivity=1.000000,Exponent=1.000000,bInvert=False))
++AxisConfig=(AxisKeyName="OculusTouch_Left_Trigger_Axis",AxisProperties=(DeadZone=0.000000,Sensitivity=1.000000,Exponent=1.000000,bInvert=False))
++AxisConfig=(AxisKeyName="OculusTouch_Left_Thumbstick_X",AxisProperties=(DeadZone=0.000000,Sensitivity=1.000000,Exponent=1.000000,bInvert=False))
++AxisConfig=(AxisKeyName="OculusTouch_Left_Thumbstick_Y",AxisProperties=(DeadZone=0.000000,Sensitivity=1.000000,Exponent=1.000000,bInvert=False))
++AxisConfig=(AxisKeyName="OculusTouch_Right_Grip_Axis",AxisProperties=(DeadZone=0.000000,Sensitivity=1.000000,Exponent=1.000000,bInvert=False))
++AxisConfig=(AxisKeyName="OculusTouch_Right_Trigger_Axis",AxisProperties=(DeadZone=0.000000,Sensitivity=1.000000,Exponent=1.000000,bInvert=False))
++AxisConfig=(AxisKeyName="OculusTouch_Right_Thumbstick_X",AxisProperties=(DeadZone=0.000000,Sensitivity=1.000000,Exponent=1.000000,bInvert=False))
++AxisConfig=(AxisKeyName="OculusTouch_Right_Thumbstick_Y",AxisProperties=(DeadZone=0.000000,Sensitivity=1.000000,Exponent=1.000000,bInvert=False))
++AxisConfig=(AxisKeyName="ValveIndex_Left_Grip_Axis",AxisProperties=(DeadZone=0.000000,Sensitivity=1.000000,Exponent=1.000000,bInvert=False))
++AxisConfig=(AxisKeyName="ValveIndex_Left_Grip_Force",AxisProperties=(DeadZone=0.000000,Sensitivity=1.000000,Exponent=1.000000,bInvert=False))
++AxisConfig=(AxisKeyName="ValveIndex_Left_Trigger_Axis",AxisProperties=(DeadZone=0.000000,Sensitivity=1.000000,Exponent=1.000000,bInvert=False))
++AxisConfig=(AxisKeyName="ValveIndex_Left_Thumbstick_X",AxisProperties=(DeadZone=0.000000,Sensitivity=1.000000,Exponent=1.000000,bInvert=False))
++AxisConfig=(AxisKeyName="ValveIndex_Left_Thumbstick_Y",AxisProperties=(DeadZone=0.000000,Sensitivity=1.000000,Exponent=1.000000,bInvert=False))
++AxisConfig=(AxisKeyName="ValveIndex_Left_Trackpad_X",AxisProperties=(DeadZone=0.000000,Sensitivity=1.000000,Exponent=1.000000,bInvert=False))
++AxisConfig=(AxisKeyName="ValveIndex_Left_Trackpad_Y",AxisProperties=(DeadZone=0.000000,Sensitivity=1.000000,Exponent=1.000000,bInvert=False))
++AxisConfig=(AxisKeyName="ValveIndex_Left_Trackpad_Force",AxisProperties=(DeadZone=0.000000,Sensitivity=1.000000,Exponent=1.000000,bInvert=False))
++AxisConfig=(AxisKeyName="ValveIndex_Left_Trackpad_Touch",AxisProperties=(DeadZone=0.000000,Sensitivity=1.000000,Exponent=1.000000,bInvert=False))
++AxisConfig=(AxisKeyName="ValveIndex_Right_Grip_Axis",AxisProperties=(DeadZone=0.000000,Sensitivity=1.000000,Exponent=1.000000,bInvert=False))
++AxisConfig=(AxisKeyName="ValveIndex_Right_Grip_Force",AxisProperties=(DeadZone=0.000000,Sensitivity=1.000000,Exponent=1.000000,bInvert=False))
++AxisConfig=(AxisKeyName="ValveIndex_Right_Trigger_Axis",AxisProperties=(DeadZone=0.000000,Sensitivity=1.000000,Exponent=1.000000,bInvert=False))
++AxisConfig=(AxisKeyName="ValveIndex_Right_Thumbstick_X",AxisProperties=(DeadZone=0.000000,Sensitivity=1.000000,Exponent=1.000000,bInvert=False))
++AxisConfig=(AxisKeyName="ValveIndex_Right_Thumbstick_Y",AxisProperties=(DeadZone=0.000000,Sensitivity=1.000000,Exponent=1.000000,bInvert=False))
++AxisConfig=(AxisKeyName="ValveIndex_Right_Trackpad_X",AxisProperties=(DeadZone=0.000000,Sensitivity=1.000000,Exponent=1.000000,bInvert=False))
++AxisConfig=(AxisKeyName="ValveIndex_Right_Trackpad_Y",AxisProperties=(DeadZone=0.000000,Sensitivity=1.000000,Exponent=1.000000,bInvert=False))
++AxisConfig=(AxisKeyName="ValveIndex_Right_Trackpad_Force",AxisProperties=(DeadZone=0.000000,Sensitivity=1.000000,Exponent=1.000000,bInvert=False))
++AxisConfig=(AxisKeyName="OculusTouch_Left_Thumbstick",AxisProperties=(DeadZone=0.000000,Sensitivity=1.000000,Exponent=1.000000,bInvert=False))
++AxisConfig=(AxisKeyName="OculusTouch_Left_FaceButton1",AxisProperties=(DeadZone=0.000000,Sensitivity=1.000000,Exponent=1.000000,bInvert=False))
++AxisConfig=(AxisKeyName="OculusTouch_Left_Trigger",AxisProperties=(DeadZone=0.000000,Sensitivity=1.000000,Exponent=1.000000,bInvert=False))
++AxisConfig=(AxisKeyName="OculusTouch_Left_FaceButton2",AxisProperties=(DeadZone=0.000000,Sensitivity=1.000000,Exponent=1.000000,bInvert=False))
++AxisConfig=(AxisKeyName="OculusTouch_Left_IndexPointing",AxisProperties=(DeadZone=0.000000,Sensitivity=1.000000,Exponent=1.000000,bInvert=False))
++AxisConfig=(AxisKeyName="OculusTouch_Left_ThumbUp",AxisProperties=(DeadZone=0.000000,Sensitivity=1.000000,Exponent=1.000000,bInvert=False))
++AxisConfig=(AxisKeyName="OculusTouch_Left_ThumbRest",AxisProperties=(DeadZone=0.000000,Sensitivity=1.000000,Exponent=1.000000,bInvert=False))
++AxisConfig=(AxisKeyName="OculusTouch_Right_Thumbstick",AxisProperties=(DeadZone=0.000000,Sensitivity=1.000000,Exponent=1.000000,bInvert=False))
++AxisConfig=(AxisKeyName="OculusTouch_Right_FaceButton1",AxisProperties=(DeadZone=0.000000,Sensitivity=1.000000,Exponent=1.000000,bInvert=False))
++AxisConfig=(AxisKeyName="OculusTouch_Right_Trigger",AxisProperties=(DeadZone=0.000000,Sensitivity=1.000000,Exponent=1.000000,bInvert=False))
++AxisConfig=(AxisKeyName="OculusTouch_Right_FaceButton2",AxisProperties=(DeadZone=0.000000,Sensitivity=1.000000,Exponent=1.000000,bInvert=False))
++AxisConfig=(AxisKeyName="OculusTouch_Right_IndexPointing",AxisProperties=(DeadZone=0.000000,Sensitivity=1.000000,Exponent=1.000000,bInvert=False))
++AxisConfig=(AxisKeyName="OculusTouch_Right_ThumbUp",AxisProperties=(DeadZone=0.000000,Sensitivity=1.000000,Exponent=1.000000,bInvert=False))
++AxisConfig=(AxisKeyName="OculusTouch_Right_ThumbRest",AxisProperties=(DeadZone=0.000000,Sensitivity=1.000000,Exponent=1.000000,bInvert=False))
++AxisConfig=(AxisKeyName="OculusHand_Left_ThumbPinchStrength",AxisProperties=(DeadZone=0.000000,Sensitivity=1.000000,Exponent=1.000000,bInvert=False))
++AxisConfig=(AxisKeyName="OculusHand_Left_IndexPinchStrength",AxisProperties=(DeadZone=0.000000,Sensitivity=1.000000,Exponent=1.000000,bInvert=False))
++AxisConfig=(AxisKeyName="OculusHand_Left_MiddlePinchStrength",AxisProperties=(DeadZone=0.000000,Sensitivity=1.000000,Exponent=1.000000,bInvert=False))
++AxisConfig=(AxisKeyName="OculusHand_Left_RingPinchStrength",AxisProperties=(DeadZone=0.000000,Sensitivity=1.000000,Exponent=1.000000,bInvert=False))
++AxisConfig=(AxisKeyName="OculusHand_Left_PinkPinchStrength",AxisProperties=(DeadZone=0.000000,Sensitivity=1.000000,Exponent=1.000000,bInvert=False))
++AxisConfig=(AxisKeyName="OculusHand_Right_ThumbPinchStrength",AxisProperties=(DeadZone=0.000000,Sensitivity=1.000000,Exponent=1.000000,bInvert=False))
++AxisConfig=(AxisKeyName="OculusHand_Right_IndexPinchStrength",AxisProperties=(DeadZone=0.000000,Sensitivity=1.000000,Exponent=1.000000,bInvert=False))
++AxisConfig=(AxisKeyName="OculusHand_Right_MiddlePinchStrength",AxisProperties=(DeadZone=0.000000,Sensitivity=1.000000,Exponent=1.000000,bInvert=False))
++AxisConfig=(AxisKeyName="OculusHand_Right_RingPinchStrength",AxisProperties=(DeadZone=0.000000,Sensitivity=1.000000,Exponent=1.000000,bInvert=False))
++AxisConfig=(AxisKeyName="OculusHand_Right_PinkPinchStrength",AxisProperties=(DeadZone=0.000000,Sensitivity=1.000000,Exponent=1.000000,bInvert=False))
++AxisConfig=(AxisKeyName="Cosmos_Left_Grip_Axis",AxisProperties=(DeadZone=0.000000,Sensitivity=1.000000,Exponent=1.000000,bInvert=False))
++AxisConfig=(AxisKeyName="Cosmos_Left_Trigger_Axis",AxisProperties=(DeadZone=0.000000,Sensitivity=1.000000,Exponent=1.000000,bInvert=False))
++AxisConfig=(AxisKeyName="Cosmos_Left_Thumbstick_X",AxisProperties=(DeadZone=0.000000,Sensitivity=1.000000,Exponent=1.000000,bInvert=False))
++AxisConfig=(AxisKeyName="Cosmos_Left_Thumbstick_Y",AxisProperties=(DeadZone=0.000000,Sensitivity=1.000000,Exponent=1.000000,bInvert=False))
++AxisConfig=(AxisKeyName="Cosmos_Right_Grip_Axis",AxisProperties=(DeadZone=0.000000,Sensitivity=1.000000,Exponent=1.000000,bInvert=False))
++AxisConfig=(AxisKeyName="Cosmos_Right_Trigger_Axis",AxisProperties=(DeadZone=0.000000,Sensitivity=1.000000,Exponent=1.000000,bInvert=False))
++AxisConfig=(AxisKeyName="Cosmos_Right_Thumbstick_X",AxisProperties=(DeadZone=0.000000,Sensitivity=1.000000,Exponent=1.000000,bInvert=False))
++AxisConfig=(AxisKeyName="Cosmos_Right_Thumbstick_Y",AxisProperties=(DeadZone=0.000000,Sensitivity=1.000000,Exponent=1.000000,bInvert=False))
+bAltEnterTogglesFullscreen=True
+bF11TogglesFullscreen=True
+bUseMouseForTouch=False
+bEnableMouseSmoothing=True
+bEnableFOVScaling=True
+bCaptureMouseOnLaunch=True
+bEnableLegacyInputScales=True
+bEnableMotionControls=True
+bFilterInputByPlatformUser=False
+bShouldFlushPressedKeysOnViewportFocusLost=True
+bAlwaysShowTouchInterface=False
+bShowConsoleOnFourFingerTap=True
+bEnableGestureRecognizer=False
+bUseAutocorrect=False
+DefaultViewportMouseCaptureMode=CapturePermanently_IncludingInitialMouseDown
+DefaultViewportMouseLockMode=LockOnCapture
+FOVScale=0.011110
+DoubleClickTime=0.200000
+DefaultPlayerInputClass=/Script/EnhancedInput.EnhancedPlayerInput
+DefaultInputComponentClass=/Script/EnhancedInput.EnhancedInputComponent
+DefaultTouchInterface=None
+-ConsoleKeys=Tilde
++ConsoleKeys=Tilde
+
+[EnhancedInputPlatformSettings_Windows EnhancedInputPlatformSettings]
+bShouldLogMappingContextRedirects=False
+
diff --git a/Templates/TP_UEIntro_BP/Config/TemplateDefs.ini b/Templates/TP_UEIntro_BP/Config/TemplateDefs.ini
new file mode 100644
--- /dev/null
+++ b/Templates/TP_UEIntro_BP/Config/TemplateDefs.ini
@@ -0,0 +1,45 @@
+﻿[/Script/GameProjectGeneration.TemplateProjectDefs]
+
+Categories=Games
+Categories=ME
+Categories=AEC
+Categories=MFG
+Categories=SIM
+
+LocalizedDisplayNames=(Language="en",Text="Intro To Unreal")
+LocalizedDescriptions=(Language="en",Text="This template takes you through the basics of the engine, including movement, adding objects to a level and the various types of asset you can author.")
+
+; This project uses a thumbnail as the project icon
+bThumbnailAsIcon=true
+
+ClassTypes=Character, SpringArmComponent, PlayerController, GameMode
+AssetTypes=Animation Sequence, Animation Blueprint, Skeleton, Skeletal Mesh, Control Rig
+
+FoldersToIgnore=Binaries
+FoldersToIgnore=Build
+FoldersToIgnore=Intermediate
+FoldersToIgnore=Saved
+FoldersToIgnore=Media
+
+FilesToIgnore="%TEMPLATENAME%.uproject"
+FilesToIgnore="%TEMPLATENAME%.png"
+FilesToIgnore="Config/TemplateDefs.ini"
+FilesToIgnore="Config/config.ini"
+FilesToIgnore="Manifest.json"
+FilesToIgnore="contents.txt"
+
+
+FolderRenames=(From="Source/%TEMPLATENAME%",To="Source/%PROJECTNAME%")
+FolderRenames=(From="Source/%TEMPLATENAME%Editor",To="Source/%PROJECTNAME%Editor")
+
+FilenameReplacements=(Extensions=("cpp","h","ini","cs"),From="%TEMPLATENAME_UPPERCASE%",To="%PROJECTNAME_UPPERCASE%",bCaseSensitive=true)
+FilenameReplacements=(Extensions=("cpp","h","ini","cs"),From="%TEMPLATENAME_LOWERCASE%",To="%PROJECTNAME_LOWERCASE%",bCaseSensitive=true)
+FilenameReplacements=(Extensions=("cpp","h","ini","cs"),From="%TEMPLATENAME%",To="%PROJECTNAME%",bCaseSensitive=false)
+
+ReplacementsInFiles=(Extensions=("cpp","h","ini","cs"),From="%TEMPLATENAME_UPPERCASE%",To="%PROJECTNAME_UPPERCASE%",bCaseSensitive=true)
+ReplacementsInFiles=(Extensions=("cpp","h","ini","cs"),From="%TEMPLATENAME_LOWERCASE%",To="%PROJECTNAME_LOWERCASE%",bCaseSensitive=true)
+ReplacementsInFiles=(Extensions=("cpp","h","ini","cs"),From="%TEMPLATENAME%",To="%PROJECTNAME%",bCaseSensitive=false)
+SharedContentPacks=(MountName="LevelPrototyping",DetailLevels=("High"))
+SharedContentPacks=(MountName="Characters",DetailLevels=("High"))
+SharedContentPacks=(MountName="Input",DetailLevels=("High"))
+EditDetailLevelPreference="High"
\ No newline at end of file
diff --git a/Templates/TP_UEIntro_BP/UEIntro.uproject b/Templates/TP_UEIntro_BP/UEIntro.uproject
new file mode 100644
--- /dev/null
+++ b/Templates/TP_UEIntro_BP/UEIntro.uproject
@@ -0,0 +1,15 @@
+{
+	"FileVersion": 3,
+	"EngineAssociation": "",
+	"Category": "",
+	"Description": "",
+	"Plugins": [
+		{
+			"Name": "ModelingToolsEditorMode",
+			"Enabled": true,
+			"TargetAllowList": [
+				"Editor"
+			]
+		}
+	]
+}
\ No newline at end of file
diff --git a/Templates/TP_UEIntro_BP/manifest.json b/Templates/TP_UEIntro_BP/manifest.json
new file mode 100644
--- /dev/null
+++ b/Templates/TP_UEIntro_BP/manifest.json
@@ -0,0 +1,78 @@
+﻿{
+	"Version": 1,
+	"Name": [
+		{
+			"Language": "en",
+			"Text": "Third Person"
+		},
+		{
+			"Language": "ko",
+			"Text": "3인칭"
+		},
+		{
+			"Language": "zh-Hans",
+			"Text": "第三人称模板"
+		},
+		{
+			"Language": "ja",
+			"Text": "三人称"
+		}
+	],
+	"Description": [
+		{
+			"Language": "en",
+			"Text": "This pack features a playable character where the camera observes an over-the-shoulder point of view. The player can rotate the camera and move around using keyboard, mouse, controller or virtual joystick on a touch device."
+		},
+		{
+			"Language": "ko",
+			"Text": "이 팩에는 카메라가 어깨 너머 시점으로 관찰하는 플레이 가능한 캐릭터가 등장합니다. 플레이어는 키보드, 마우스, 컨트롤러 또는 터치 디바이스의 가상 조이스틱을 사용하여 카메라를 회전하고 이동할 수 있습니다."
+		},
+		{
+			"Language": "zh-Hans",
+			"Text": "此模板包含一个可操控的角色，镜头采用过肩视角进行拍摄。玩家可以通过键盘、鼠标、手柄或触控设备上的虚拟摇杆来旋转镜头并进行移动。"
+		},
+		{
+			"Language": "ja",
+			"Text": "このパックには、カメラが肩越しの視点で表示されるプレイアブル キャラクターが含まれています。プレイヤーは、キーボード、マウス、コントローラー、またはタッチ デバイスのバーチャル ジョイスティックを使用して、カメラを回転させたり移動したりできます"
+		}
+	],
+	"AssetTypes": [
+		{
+			"Language": "en",
+			"Text": "Animation Sequence, Animation Blueprint, Skeleton, Skeletal Mesh, Control Rig"
+		}
+	],
+	"SearchTags": [
+		{
+			"Language": "en",
+			"Text": "Action, Adventure, Third"
+		}
+	],
+	"ClassTypes": "Character, SpringArmComponent, PlayerController, GameMode",
+	"Category": "BlueprintFeature",
+	"FocusAsset": "/Game/ThirdPerson/Blueprints/BP_ThirdPersonCharacter.BP_ThirdPersonCharacter",
+	"Thumbnail": "TP_ThirdPersonBP.png",
+	"Screenshots": [
+		"TP_ThirdPersonBP_Preview.png"
+	],
+	"AdditionalFeaturePacks": [
+		{
+			"MountName": "LevelPrototyping",
+			"DetailLevels": [
+				"High"
+			]
+		},
+		{
+			"MountName": "Characters",
+			"DetailLevels": [
+				"High"
+			]
+		},
+		{
+			"MountName": "Input",
+			"DetailLevels": [
+				"High"
+			]
+		}
+	]
+}
