Updates to UE Intro Template

#virtualized

diff --git a/Templates/TP_UEIntro_BP/Config/DefaultEngine.ini b/Templates/TP_UEIntro_BP/Config/DefaultEngine.ini
--- a/Templates/TP_UEIntro_BP/Config/DefaultEngine.ini
+++ b/Templates/TP_UEIntro_BP/Config/DefaultEngine.ini
@@ -50,6 +50,10 @@ r.Substrate=False
 r.Refraction.Blur=False
 r.Lumen.HardwareRayTracing=True
 r.Lumen.Reflections.HardwareRayTracing.Translucent.Refraction.EnableForProject=False
+r.ReflectionCaptureResolution=256
+r.Lumen.TraceMeshSDFs=0
+r.Lumen.TranslucencyReflections.FrontLayer.EnableForProject=False
+r.DistanceFields.DefaultVoxelDensity=0.200000
 
 [/Script/WindowsTargetPlatform.WindowsTargetSettings]
 DefaultGraphicsRHI=DefaultGraphicsRHI_DX12
diff --git a/Templates/TP_UEIntro_BP/Config/DefaultInput.ini b/Templates/TP_UEIntro_BP/Config/DefaultInput.ini
--- a/Templates/TP_UEIntro_BP/Config/DefaultInput.ini
+++ b/Templates/TP_UEIntro_BP/Config/DefaultInput.ini
@@ -105,8 +105,8 @@ bAlwaysShowTouchInterface=False
 bShowConsoleOnFourFingerTap=True
 bEnableGestureRecognizer=False
 bUseAutocorrect=False
-DefaultViewportMouseCaptureMode=CapturePermanently_IncludingInitialMouseDown
-DefaultViewportMouseLockMode=LockOnCapture
+DefaultViewportMouseCaptureMode=CapturePermanently
+DefaultViewportMouseLockMode=LockAlways
 FOVScale=0.011110
 DoubleClickTime=0.200000
 DefaultPlayerInputClass=/Script/EnhancedInput.EnhancedPlayerInput
diff --git a/Templates/TP_UEIntro_BP/UEIntro.uproject b/Templates/TP_UEIntro_BP/UEIntro.uproject
--- a/Templates/TP_UEIntro_BP/UEIntro.uproject
+++ b/Templates/TP_UEIntro_BP/UEIntro.uproject
@@ -10,6 +10,10 @@
 			"TargetAllowList": [
 				"Editor"
 			]
+		},
+		{
+			"Name": "Landmass",
+			"Enabled": true
 		}
 	]
 }
\ No newline at end of file
