Updated Intro Template

#virtualized

diff --git a/Templates/TP_UEIntro_BP/Config/TemplateDefs.ini b/Templates/TP_UEIntro_BP/Config/TemplateDefs.ini
--- a/Templates/TP_UEIntro_BP/Config/TemplateDefs.ini
+++ b/Templates/TP_UEIntro_BP/Config/TemplateDefs.ini
@@ -6,8 +6,10 @@ Categories=AEC
 Categories=MFG
 Categories=SIM
 
+SortKey="0"
+
 LocalizedDisplayNames=(Language="en",Text="Intro To Unreal")
-LocalizedDescriptions=(Language="en",Text="This template takes you through the basics of the engine, including movement, adding objects to a level and the various types of asset you can author.")
+LocalizedDescriptions=(Language="en",Text="This template takes you through the basics of the engine, including movement, adding objects to a level and the various types of asset you can create.")
 
 ; This project uses a thumbnail as the project icon
 bThumbnailAsIcon=true
