Updated intro template
added more text
cleaned up content

#virtualized

diff --git a/Templates/TP_UEIntro_BP/Config/DefaultEngine.ini b/Templates/TP_UEIntro_BP/Config/DefaultEngine.ini
--- a/Templates/TP_UEIntro_BP/Config/DefaultEngine.ini
+++ b/Templates/TP_UEIntro_BP/Config/DefaultEngine.ini
@@ -50,7 +50,7 @@ r.Substrate=False
 r.Refraction.Blur=False
 r.Lumen.HardwareRayTracing=True
 r.Lumen.Reflections.HardwareRayTracing.Translucent.Refraction.EnableForProject=False
-r.ReflectionCaptureResolution=256
+r.ReflectionCaptureResolution=512
 r.Lumen.TraceMeshSDFs=0
 r.Lumen.TranslucencyReflections.FrontLayer.EnableForProject=False
 r.DistanceFields.DefaultVoxelDensity=0.200000
