Configure final tutorial URL for TP_UEIntro_BP project template.

diff --git a/Templates/TP_UEIntro_BP/Config/DefaultEditorPerProjectUserSettings.ini b/Templates/TP_UEIntro_BP/Config/DefaultEditorPerProjectUserSettings.ini
--- a/Templates/TP_UEIntro_BP/Config/DefaultEditorPerProjectUserSettings.ini
+++ b/Templates/TP_UEIntro_BP/Config/DefaultEditorPerProjectUserSettings.ini
@@ -1,2 +1,5 @@
 [ContentBrowser]
-ContentBrowserTab1.SelectedPaths=/Game/FirstPersonBP
\ No newline at end of file
+ContentBrowserTab1.SelectedPaths=/Game/FirstPersonBP
+
+[/Script/InEditorDocumentation.InEditorDocumentationSettings]
+TutorialUrl="https://dev.epicgames.com/documentation/unreal-engine/introduction-to-unreal-engine"
\ No newline at end of file
