Fixed launch character on jump pad
fixed character reflection issue on manny + quinn
collecting localization text for ue intro

#virtualized

diff --git a/Templates/TP_UEIntro_BP/Config/DefaultEngine.ini b/Templates/TP_UEIntro_BP/Config/DefaultEngine.ini
--- a/Templates/TP_UEIntro_BP/Config/DefaultEngine.ini
+++ b/Templates/TP_UEIntro_BP/Config/DefaultEngine.ini
@@ -54,6 +54,7 @@ r.ReflectionCaptureResolution=512
 r.Lumen.TraceMeshSDFs=0
 r.Lumen.TranslucencyReflections.FrontLayer.EnableForProject=False
 r.DistanceFields.DefaultVoxelDensity=0.200000
+r.Lumen.HardwareRayTracing.LightingMode=0
 
 [/Script/WindowsTargetPlatform.WindowsTargetSettings]
 DefaultGraphicsRHI=DefaultGraphicsRHI_DX12
