added additional languages for UE intro template

diff --git a/Templates/TP_UEIntro_BP/Config/Localization/Game_Compile.ini b/Templates/TP_UEIntro_BP/Config/Localization/Game_Compile.ini
--- a/Templates/TP_UEIntro_BP/Config/Localization/Game_Compile.ini
+++ b/Templates/TP_UEIntro_BP/Config/Localization/Game_Compile.ini
@@ -12,6 +12,19 @@ bValidateSafeWhitespace=false
 bValidateRichTextTags=false
 NativeCulture=en
 CulturesToGenerate=en
+CulturesToGenerate=ar
+CulturesToGenerate=ko
+CulturesToGenerate=ja
+CulturesToGenerate=zh-Hans
+CulturesToGenerate=fr
+CulturesToGenerate=de
+CulturesToGenerate=it
+CulturesToGenerate=es
+CulturesToGenerate=es-419
+CulturesToGenerate=pl
+CulturesToGenerate=pt-BR
+CulturesToGenerate=ru
+CulturesToGenerate=tr
 
 [GatherTextStep0]
 CommandletClass=GenerateTextLocalizationResource
diff --git a/Templates/TP_UEIntro_BP/Config/Localization/Game_Export.ini b/Templates/TP_UEIntro_BP/Config/Localization/Game_Export.ini
--- a/Templates/TP_UEIntro_BP/Config/Localization/Game_Export.ini
+++ b/Templates/TP_UEIntro_BP/Config/Localization/Game_Export.ini
@@ -5,6 +5,19 @@ SourcePath=Content/Localization/Game
 DestinationPath=Content/Localization/Game
 NativeCulture=en
 CulturesToGenerate=en
+CulturesToGenerate=ar
+CulturesToGenerate=ko
+CulturesToGenerate=ja
+CulturesToGenerate=zh-Hans
+CulturesToGenerate=fr
+CulturesToGenerate=de
+CulturesToGenerate=it
+CulturesToGenerate=es
+CulturesToGenerate=es-419
+CulturesToGenerate=pl
+CulturesToGenerate=pt-BR
+CulturesToGenerate=ru
+CulturesToGenerate=tr
 ManifestName=Game.manifest
 ArchiveName=Game.archive
 PortableObjectName=Game.po
diff --git a/Templates/TP_UEIntro_BP/Config/Localization/Game_ExportDialogueScript.ini b/Templates/TP_UEIntro_BP/Config/Localization/Game_ExportDialogueScript.ini
--- a/Templates/TP_UEIntro_BP/Config/Localization/Game_ExportDialogueScript.ini
+++ b/Templates/TP_UEIntro_BP/Config/Localization/Game_ExportDialogueScript.ini
@@ -5,6 +5,19 @@ SourcePath=Content/Localization/Game
 DestinationPath=Content/Localization/Game
 NativeCulture=en
 CulturesToGenerate=en
+CulturesToGenerate=ar
+CulturesToGenerate=ko
+CulturesToGenerate=ja
+CulturesToGenerate=zh-Hans
+CulturesToGenerate=fr
+CulturesToGenerate=de
+CulturesToGenerate=it
+CulturesToGenerate=es
+CulturesToGenerate=es-419
+CulturesToGenerate=pl
+CulturesToGenerate=pt-BR
+CulturesToGenerate=ru
+CulturesToGenerate=tr
 ManifestName=Game.manifest
 ArchiveName=Game.archive
 DialogueScriptName=GameDialogue.csv
diff --git a/Templates/TP_UEIntro_BP/Config/Localization/Game_Gather.ini b/Templates/TP_UEIntro_BP/Config/Localization/Game_Gather.ini
--- a/Templates/TP_UEIntro_BP/Config/Localization/Game_Gather.ini
+++ b/Templates/TP_UEIntro_BP/Config/Localization/Game_Gather.ini
@@ -9,6 +9,19 @@ ManifestName=Game.manifest
 ArchiveName=Game.archive
 NativeCulture=en
 CulturesToGenerate=en
+CulturesToGenerate=ar
+CulturesToGenerate=ko
+CulturesToGenerate=ja
+CulturesToGenerate=zh-Hans
+CulturesToGenerate=fr
+CulturesToGenerate=de
+CulturesToGenerate=it
+CulturesToGenerate=es
+CulturesToGenerate=es-419
+CulturesToGenerate=pl
+CulturesToGenerate=pt-BR
+CulturesToGenerate=ru
+CulturesToGenerate=tr
 
 [GatherTextStep0]
 CommandletClass=GatherTextFromAssets
diff --git a/Templates/TP_UEIntro_BP/Config/Localization/Game_GenerateReports.ini b/Templates/TP_UEIntro_BP/Config/Localization/Game_GenerateReports.ini
--- a/Templates/TP_UEIntro_BP/Config/Localization/Game_GenerateReports.ini
+++ b/Templates/TP_UEIntro_BP/Config/Localization/Game_GenerateReports.ini
@@ -6,6 +6,19 @@ DestinationPath=Content/Localization/Game
 ManifestName=Game.manifest
 ArchiveName=Game.archive
 CulturesToGenerate=en
+CulturesToGenerate=ar
+CulturesToGenerate=ko
+CulturesToGenerate=ja
+CulturesToGenerate=zh-Hans
+CulturesToGenerate=fr
+CulturesToGenerate=de
+CulturesToGenerate=it
+CulturesToGenerate=es
+CulturesToGenerate=es-419
+CulturesToGenerate=pl
+CulturesToGenerate=pt-BR
+CulturesToGenerate=ru
+CulturesToGenerate=tr
 
 [GatherTextStep0]
 CommandletClass=GenerateTextLocalizationReport
diff --git a/Templates/TP_UEIntro_BP/Config/Localization/Game_Import.ini b/Templates/TP_UEIntro_BP/Config/Localization/Game_Import.ini
--- a/Templates/TP_UEIntro_BP/Config/Localization/Game_Import.ini
+++ b/Templates/TP_UEIntro_BP/Config/Localization/Game_Import.ini
@@ -5,6 +5,19 @@ SourcePath=Content/Localization/Game
 DestinationPath=Content/Localization/Game
 NativeCulture=en
 CulturesToGenerate=en
+CulturesToGenerate=ar
+CulturesToGenerate=ko
+CulturesToGenerate=ja
+CulturesToGenerate=zh-Hans
+CulturesToGenerate=fr
+CulturesToGenerate=de
+CulturesToGenerate=it
+CulturesToGenerate=es
+CulturesToGenerate=es-419
+CulturesToGenerate=pl
+CulturesToGenerate=pt-BR
+CulturesToGenerate=ru
+CulturesToGenerate=tr
 ManifestName=Game.manifest
 ArchiveName=Game.archive
 PortableObjectName=Game.po
diff --git a/Templates/TP_UEIntro_BP/Config/Localization/Game_ImportDialogue.ini b/Templates/TP_UEIntro_BP/Config/Localization/Game_ImportDialogue.ini
--- a/Templates/TP_UEIntro_BP/Config/Localization/Game_ImportDialogue.ini
+++ b/Templates/TP_UEIntro_BP/Config/Localization/Game_ImportDialogue.ini
@@ -6,6 +6,19 @@ ManifestName=Game.manifest
 ArchiveName=Game.archive
 NativeCulture=en
 CulturesToGenerate=en
+CulturesToGenerate=ar
+CulturesToGenerate=ko
+CulturesToGenerate=ja
+CulturesToGenerate=zh-Hans
+CulturesToGenerate=fr
+CulturesToGenerate=de
+CulturesToGenerate=it
+CulturesToGenerate=es
+CulturesToGenerate=es-419
+CulturesToGenerate=pl
+CulturesToGenerate=pt-BR
+CulturesToGenerate=ru
+CulturesToGenerate=tr
 
 [GatherTextStep0]
 CommandletClass=ImportLocalizedDialogue
diff --git a/Templates/TP_UEIntro_BP/Config/Localization/Game_ImportDialogueScript.ini b/Templates/TP_UEIntro_BP/Config/Localization/Game_ImportDialogueScript.ini
--- a/Templates/TP_UEIntro_BP/Config/Localization/Game_ImportDialogueScript.ini
+++ b/Templates/TP_UEIntro_BP/Config/Localization/Game_ImportDialogueScript.ini
@@ -5,6 +5,19 @@ SourcePath=Content/Localization/Game
 DestinationPath=Content/Localization/Game
 NativeCulture=en
 CulturesToGenerate=en
+CulturesToGenerate=ar
+CulturesToGenerate=ko
+CulturesToGenerate=ja
+CulturesToGenerate=zh-Hans
+CulturesToGenerate=fr
+CulturesToGenerate=de
+CulturesToGenerate=it
+CulturesToGenerate=es
+CulturesToGenerate=es-419
+CulturesToGenerate=pl
+CulturesToGenerate=pt-BR
+CulturesToGenerate=ru
+CulturesToGenerate=tr
 ManifestName=Game.manifest
 ArchiveName=Game.archive
 DialogueScriptName=GameDialogue.csv
